Disabled Oculus and SteamVR plugins in the nDisplay template to prevent plugins collisions

diff --git a/Templates/TP_nDisplayBP/TP_nDisplayBP.uproject b/Templates/TP_nDisplayBP/TP_nDisplayBP.uproject
--- a/Templates/TP_nDisplayBP/TP_nDisplayBP.uproject
+++ b/Templates/TP_nDisplayBP/TP_nDisplayBP.uproject
@@ -7,6 +7,25 @@
 		{
 			"Name": "nDisplay",
 			"Enabled": true
+		},
+		{
+			"Name": "OculusVR",
+			"Enabled": false,
+			"SupportedTargetPlatforms": [
+				"Win32",
+				"Win64",
+				"Android"
+			]
+		},
+		{
+			"Name": "SteamVR",
+			"Enabled": false,
+			"SupportedTargetPlatforms": [
+				"Win32",
+				"Win64",
+				"Linux",
+				"Mac"
+			]
 		}
 	],
 	"TargetPlatforms": [
