Fixed assets versioning and uproject template requirements

diff --git a/Templates/TP_nDisplayBP/TP_nDisplayBP.uproject b/Templates/TP_nDisplayBP/TP_nDisplayBP.uproject
--- a/Templates/TP_nDisplayBP/TP_nDisplayBP.uproject
+++ b/Templates/TP_nDisplayBP/TP_nDisplayBP.uproject
@@ -6,7 +6,10 @@
 	"Plugins": [
 		{
 			"Name": "nDisplay",
-			"Enabled": true
+			"Enabled": true,
+			"SupportedTargetPlatforms": [
+				"Win64"
+			]			
 		},
 		{
 			"Name": "OculusVR",
