TP_nDisplayBP: Added LinuxEditor and WindowsEditor target platforms.

diff --git a/Templates/TP_nDisplayBP/TP_nDisplayBP.uproject b/Templates/TP_nDisplayBP/TP_nDisplayBP.uproject
--- a/Templates/TP_nDisplayBP/TP_nDisplayBP.uproject
+++ b/Templates/TP_nDisplayBP/TP_nDisplayBP.uproject
@@ -19,6 +19,8 @@
 	],
 	"TargetPlatforms": [
 		"Windows",
-		"Linux"
+		"WindowsEditor",
+		"Linux",
+		"LinuxEditor"
 	]
 }
\ No newline at end of file
